Add EditorSysConfigAssistant plugin to help guide users to system configuration changes that can improve development experience.

diff --git a/Engine/Plugins/Editor/EditorSysConfigAssistant/EditorSysConfigAssistant.uplugin b/Engine/Plugins/Editor/EditorSysConfigAssistant/EditorSysConfigAssistant.uplugin
new file mode 100644
--- /dev/null
+++ b/Engine/Plugins/Editor/EditorSysConfigAssistant/EditorSysConfigAssistant.uplugin
@@ -0,0 +1,25 @@
+{
+	"FileVersion": 3,
+	"Version": 1,
+	"VersionName": "0.1",
+	"FriendlyName": "Editor System Configuration Assistant",
+	"Description": "Editor utility for offering system configuration suggestions",
+	"Category": "Other",
+	"CreatedBy": "Epic Games, Inc.",
+	"CreatedByURL": "https://epicgames.com",
+	"DocsURL": "",
+	"MarketplaceURL": "",
+	"SupportURL": "",
+	"CanContainContent": true,
+	"IsBetaVersion": false,
+	"IsExperimentalVersion": true,
+	"Installed": false,
+	"EnabledByDefault": false,
+	"Modules": [
+		{
+			"Name": "EditorSysConfigAssistant",
+			"Type": "Editor",
+			"LoadingPhase": "Default"
+		}
+	]
+}
\ No newline at end of file
diff --git a/Engine/Plugins/Editor/EditorSysConfigAssistant/Source/EditorSysConfigAssistant/EditorSysConfigAssistant.Build.cs b/Engine/Plugins/Editor/EditorSysConfigAssistant/Source/EditorSysConfigAssistant/EditorSysConfigAssistant.Build.cs
new file mode 100644
--- /dev/null
+++ b/Engine/Plugins/Editor/EditorSysConfigAssistant/Source/EditorSysConfigAssistant/EditorSysConfigAssistant.Build.cs
@@ -0,0 +1,62 @@
+// Copyright Epic Games, Inc. All Rights Reserved.
+
+using UnrealBuildTool;
+
+public class EditorSysConfigAssistant : ModuleRules
+{
+	public EditorSysConfigAssistant(ReadOnlyTargetRules Target) : base(Target)
+	{
+		PCHUsage = ModuleRules.PCHUsageMode.UseExplicitOrSharedPCHs;
+		
+		PublicIncludePaths.AddRange(
+			new string[] {
+				// ... add public include paths required here ...
+			}
+			);
+				
+		
+		PrivateIncludePaths.AddRange(
+			new string[] {
+				// ... add other private include paths required here ...
+			}
+			);
+			
+		
+		PublicDependencyModuleNames.AddRange(
+			new string[]
+			{
+				"Core",
+				
+				"UnrealEd",
+				"InputCore",
+				"InteractiveToolsFramework",
+				
+				// ... add other public dependencies that you statically link with here ...
+			}
+			);
+			
+		
+		PrivateDependencyModuleNames.AddRange(
+			new string[]
+			{
+				"CoreUObject",
+				"Engine",
+				"Slate",
+				"SlateCore",
+				"ToolWidgets",
+				"WorkspaceMenuStructure",
+				
+				"EditorSubsystem",
+				// ... add private dependencies that you statically link with here ...	
+			}
+			);
+		
+		
+		DynamicallyLoadedModuleNames.AddRange(
+			new string[]
+			{
+				// ... add any modules that your module loads dynamically here ...
+			}
+			);
+	}
+}
diff --git a/Engine/Plugins/Editor/EditorSysConfigAssistant/Source/EditorSysConfigAssistant/Private/EditorSysConfigAssistantModule.cpp b/Engine/Plugins/Editor/EditorSysConfigAssistant/Source/EditorSysConfigAssistant/Private/EditorSysConfigAssistantModule.cpp
new file mode 100644
--- /dev/null
+++ b/Engine/Plugins/Editor/EditorSysConfigAssistant/Source/EditorSysConfigAssistant/Private/EditorSysConfigAssistantModule.cpp
@@ -0,0 +1,74 @@
+// Copyright Epic Games, Inc. All Rights Reserved.
+
+#include "EditorSysConfigAssistantModule.h"
+
+#include "Features/IModularFeatures.h"
+#include "Framework/Docking/TabManager.h"
+#include "Styling/AppStyle.h"
+#include "Textures/SlateIcon.h"
+#include "Widgets/DeclarativeSyntaxSupport.h"
+#include "Widgets/Docking/SDockTab.h"
+#include "Widgets/SEditorSysConfigAssistant.h"
+#include "WorkspaceMenuStructure.h"
+#include "WorkspaceMenuStructureModule.h"
+
+#if PLATFORM_WINDOWS
+#include "Windows/EditorSysConfigFeatureLastAccessTime.h"
+#endif
+
+static const FName EditorSysConfigAssistantTabName("EditorSysConfigAssistant");
+
+class FEditorSysConfigAssistantModule : public IEditorSysConfigAssistantModule
+{
+public:
+	FEditorSysConfigAssistantModule() = default;
+	~FEditorSysConfigAssistantModule() = default;
+
+	// ~Begin IModuleInterface
+	virtual void StartupModule() override;
+	virtual void ShutdownModule() override;
+	// ~End IModuleInterface
+
+	virtual void ShowSystemConfigAssistant() override;
+private:
+#if PLATFORM_WINDOWS
+	FEditorSysConfigFeatureLastAccessTime LastAccessTimeFeature;
+#endif // PLATFORM_WINDOWS
+};
+
+static TSharedRef<SDockTab> SpawnEditorSysConfigAssistantTab(const FSpawnTabArgs& SpawnTabArgs)
+{
+	const TSharedRef<SDockTab> DockTab = SNew(SDockTab)
+		.TabRole(ETabRole::NomadTab);
+
+	DockTab->SetContent(SNew(SEditorSysConfigAssistant, DockTab, SpawnTabArgs.GetOwnerWindow()));
+
+	return DockTab;
+}
+
+void FEditorSysConfigAssistantModule::StartupModule()
+{
+	FGlobalTabmanager::Get()->RegisterNomadTabSpawner(EditorSysConfigAssistantTabName, FOnSpawnTab::CreateStatic(&SpawnEditorSysConfigAssistantTab))
+		.SetDisplayName(NSLOCTEXT("EditorSysConfigAssistant", "EditorSysConfigAssistantTabTitle", "System Config"))
+		.SetTooltipText(NSLOCTEXT("EditorSysConfigAssistant", "EditorSysConfigAssistantTooltipText", "Open the System Config Assistant."))
+		.SetIcon(FSlateIcon(FAppStyle::GetAppStyleSetName(), "EditorPreferences.TabIcon"))
+		.SetGroup(WorkspaceMenu::GetMenuStructure().GetToolsCategory());
+
+#if PLATFORM_WINDOWS
+	IModularFeatures::Get().RegisterModularFeature(LastAccessTimeFeature.GetModularFeatureName(), &LastAccessTimeFeature);
+#endif // PLATFORM_WINDOWS
+}
+
+void FEditorSysConfigAssistantModule::ShutdownModule()
+{
+#if PLATFORM_WINDOWS
+	IModularFeatures::Get().UnregisterModularFeature(LastAccessTimeFeature.GetModularFeatureName(), &LastAccessTimeFeature);
+#endif // PLATFORM_WINDOWS
+}
+
+void FEditorSysConfigAssistantModule::ShowSystemConfigAssistant()
+{
+	FGlobalTabmanager::Get()->TryInvokeTab(EditorSysConfigAssistantTabName);
+}
+
+IMPLEMENT_MODULE(FEditorSysConfigAssistantModule, EditorSysConfigAssistant)
diff --git a/Engine/Plugins/Editor/EditorSysConfigAssistant/Source/EditorSysConfigAssistant/Private/EditorSysConfigAssistantSubsystem.cpp b/Engine/Plugins/Editor/EditorSysConfigAssistant/Source/EditorSysConfigAssistant/Private/EditorSysConfigAssistantSubsystem.cpp
new file mode 100644
--- /dev/null
+++ b/Engine/Plugins/Editor/EditorSysConfigAssistant/Source/EditorSysConfigAssistant/Private/EditorSysConfigAssistantSubsystem.cpp
@@ -0,0 +1,354 @@
+// Copyright Epic Games, Inc. All Rights Reserved.
+
+#include "EditorSysConfigAssistantSubsystem.h"
+
+#include "Async/Async.h"
+#include "EditorSysConfigAssistantModule.h"
+#include "EditorSysConfigFeature.h"
+#include "Features/IModularFeatures.h"
+#include "Framework/Application/SlateApplication.h"
+#include "Framework/Notifications/NotificationManager.h"
+#include "Internationalization/Internationalization.h"
+#include "Internationalization/Text.h"
+#include "Misc/App.h"
+#include "Misc/FileHelper.h"
+#include "Misc/Paths.h"
+#include "UnrealEdMisc.h"
+#include "Widgets/Notifications/SNotificationList.h"
+
+#if PLATFORM_WINDOWS
+#include "Microsoft/MinimalWindowsApi.h"
+#include "Windows/WindowsPlatformMisc.h"
+#include "Windows/AllowWindowsPlatformTypes.h"
+THIRD_PARTY_INCLUDES_START
+#include <winreg.h>
+THIRD_PARTY_INCLUDES_END
+#include "Windows/HideWindowsPlatformTypes.h"
+#endif
+
+#include UE_INLINE_GENERATED_CPP_BY_NAME(EditorSysConfigAssistantSubsystem)
+
+#define LOCTEXT_NAMESPACE "EditorSysConfigAssistant"
+
+
+static bool IsUsingCookedEditorContent()
+{
+	return FPaths::FileExists(*FPaths::Combine(FPaths::ProjectDir(), TEXT("EditorClientAssetRegistry.bin")));
+}
+
+bool UEditorSysConfigAssistantSubsystem::ShouldCreateSubsystem(UObject* Outer) const
+{
+#if IS_PROGRAM || !WITH_EDITOR
+	return false;
+#endif
+	
+	if (FApp::IsUnattended() || IsRunningCommandlet() || GIsRunningUnattendedScript || IsUsingCookedEditorContent() || !FSlateApplication::IsInitialized())
+	{
+		return false;
+	}
+	
+	return Super::ShouldCreateSubsystem(Outer);
+}
+
+void UEditorSysConfigAssistantSubsystem::Initialize(FSubsystemCollectionBase& Collection)
+{
+	Super::Initialize(Collection);
+
+	IModularFeatures& ModularFeatures = IModularFeatures::Get();
+	IModularFeatures::FScopedLockModularFeatureList LockModularFeatureList;
+	TArray<IEditorSysConfigFeature*> SysConfigFeatures = IModularFeatures::Get().GetModularFeatureImplementations<IEditorSysConfigFeature>(IEditorSysConfigFeature::GetModularFeatureName());
+	for (IEditorSysConfigFeature* SysConfigFeature : SysConfigFeatures)
+	{
+		SysConfigFeature->StartSystemCheck();
+	}
+
+	ModularFeatures.OnModularFeatureRegistered().AddUObject(this, &UEditorSysConfigAssistantSubsystem::HandleModularFeatureRegistered);
+	ModularFeatures.OnModularFeatureUnregistered().AddUObject(this, &UEditorSysConfigAssistantSubsystem::HandleModularFeatureUnregistered);
+}
+
+void UEditorSysConfigAssistantSubsystem::Deinitialize()
+{
+	Super::Deinitialize();
+
+	IModularFeatures& ModularFeatures = IModularFeatures::Get();
+	ModularFeatures.OnModularFeatureRegistered().RemoveAll(this);
+	ModularFeatures.OnModularFeatureUnregistered().RemoveAll(this);
+}
+
+void UEditorSysConfigAssistantSubsystem::HandleModularFeatureRegistered(const FName& InFeatureName, IModularFeature* InFeature)
+{
+	if (InFeatureName == IEditorSysConfigFeature::GetModularFeatureName())
+	{
+		IEditorSysConfigFeature* NewSysConfigFeature = static_cast<IEditorSysConfigFeature*>(InFeature);
+		NewSysConfigFeature->StartSystemCheck();
+	}
+}
+
+void UEditorSysConfigAssistantSubsystem::HandleModularFeatureUnregistered(const FName& InFeatureName, IModularFeature* InFeature)
+{
+	if (InFeatureName == IEditorSysConfigFeature::GetModularFeatureName())
+	{
+		IEditorSysConfigFeature* RemovedSysConfigFeature = static_cast<IEditorSysConfigFeature*>(InFeature);
+		FWriteScopeLock _(IssuesLock);
+		Issues.RemoveAll([RemovedSysConfigFeature](const TSharedPtr<FEditorSysConfigIssue>& Issue)
+			{
+				return Issue->Feature == RemovedSysConfigFeature;
+			});
+	}
+}
+
+void UEditorSysConfigAssistantSubsystem::NotifySystemConfigIssues()
+{
+	check(IsInGameThread());
+	
+	UEditorSysConfigAssistantSubsystem* Subsystem = GEditor->GetEditorSubsystem<UEditorSysConfigAssistantSubsystem>();
+	if (!Subsystem)
+	{
+		return;
+	}
+
+	FNotificationInfo Info(LOCTEXT("SystemConfigurationIssuesExist", "System configuration is impacting editor experience"));
+	Info.bUseSuccessFailIcons = true;
+	Info.Image = FAppStyle::GetBrush(TEXT("MessageLog.Warning"));
+	Info.bFireAndForget = false;
+	Info.bUseThrobber = true;
+	Info.FadeOutDuration = 0.0f;
+	Info.ExpireDuration = 0.0f;
+
+	Info.ButtonDetails.Add(FNotificationButtonInfo(LOCTEXT("OpenSysConfig", "System Config"), FText(), FSimpleDelegate::CreateLambda([]() {
+		IEditorSysConfigAssistantModule::Get().ShowSystemConfigAssistant();
+	}),
+		SNotificationItem::ECompletionState::CS_Fail));
+
+
+	Subsystem->IssueNotificationItem = FSlateNotificationManager::Get().AddNotification(Info);
+
+	if (Subsystem->IssueNotificationItem.IsValid())
+	{
+		Subsystem->IssueNotificationItem->SetCompletionState(SNotificationItem::CS_Fail);
+	}
+}
+
+void UEditorSysConfigAssistantSubsystem::NotifyRestart(bool bApplicationOnly)
+{
+	TSharedPtr<SNotificationItem> NotificationPin = RestartNotificationItem.Pin();
+	if (NotificationPin.IsValid())
+	{
+		return;
+	}
+	
+	FText RestartText;
+	if (bApplicationOnly)
+	{
+		RestartText = LOCTEXT("RestartRequiredTitle", "Application restart required to apply new system configuration");
+	}
+	else
+	{
+		RestartText = LOCTEXT("RestartRequiredTitle", "System restart required to apply new system configuration");
+	}
+	FNotificationInfo Info(RestartText);
+
+	// Add the buttons with text, tooltip and callback
+	Info.ButtonDetails.Add(FNotificationButtonInfo(
+		LOCTEXT("RestartNow", "Restart Now"), 
+		LOCTEXT("RestartNowToolTip", "Restart now to finish applying your new system configuration."), 
+		FSimpleDelegate::CreateUObject(this, bApplicationOnly ? &UEditorSysConfigAssistantSubsystem::OnApplicationRestartClicked : &UEditorSysConfigAssistantSubsystem::OnSystemRestartClicked))
+		);
+	Info.ButtonDetails.Add(FNotificationButtonInfo(
+		LOCTEXT("RestartLater", "Restart Later"), 
+		LOCTEXT("RestartLaterToolTip", "Dismiss this notificaton without restarting. Some new system configuration will not be applied."), 
+		FSimpleDelegate::CreateUObject(this, &UEditorSysConfigAssistantSubsystem::OnRestartDismissClicked))
+		);
+
+	// We will be keeping track of this ourselves
+	Info.bFireAndForget = false;
+
+	// Set the width so that the notification doesn't resize as its text changes
+	Info.WidthOverride = 300.0f;
+
+	Info.bUseLargeFont = false;
+	Info.bUseThrobber = false;
+	Info.bUseSuccessFailIcons = false;
+
+	// Present notification
+	RestartNotificationItem = FSlateNotificationManager::Get().AddNotification(Info);
+	NotificationPin = RestartNotificationItem.Pin();
+
+	if (NotificationPin.IsValid())
+	{
+		NotificationPin->SetCompletionState(SNotificationItem::CS_Pending);
+	}
+}
+
+void UEditorSysConfigAssistantSubsystem::OnApplicationRestartClicked()
+{
+	TSharedPtr<SNotificationItem> NotificationPin = RestartNotificationItem.Pin();
+	if (NotificationPin.IsValid())
+	{
+		NotificationPin->SetText(LOCTEXT("RestartingNow", "Restarting..."));
+		NotificationPin->SetCompletionState(SNotificationItem::CS_Success);
+		NotificationPin->ExpireAndFadeout();
+		RestartNotificationItem.Reset();
+	}
+
+	const bool bWarn = false;
+	FUnrealEdMisc::Get().RestartEditor(bWarn);
+}
+
+void UEditorSysConfigAssistantSubsystem::OnSystemRestartClicked()
+{
+	TSharedPtr<SNotificationItem> NotificationPin = RestartNotificationItem.Pin();
+	if (NotificationPin.IsValid())
+	{
+		NotificationPin->SetText(LOCTEXT("RestartingNow", "Restarting..."));
+		NotificationPin->SetCompletionState(SNotificationItem::CS_Success);
+		NotificationPin->ExpireAndFadeout();
+		RestartNotificationItem.Reset();
+	}
+
+#if PLATFORM_WINDOWS
+	::ExitWindowsEx(EWX_REBOOT | EWX_FORCE, SHTDN_REASON_MINOR_MAINTENANCE | SHTDN_REASON_FLAG_PLANNED);
+#endif
+	RequestEngineExit(TEXT("Restarting system to apply system config changes"));
+
+}
+
+void UEditorSysConfigAssistantSubsystem::OnRestartDismissClicked()
+{
+	TSharedPtr<SNotificationItem> NotificationPin = RestartNotificationItem.Pin();
+	if (NotificationPin.IsValid())
+	{
+		NotificationPin->SetText(LOCTEXT("RestartDismissed", "Restart Dismissed..."));
+		NotificationPin->SetCompletionState(SNotificationItem::CS_None);
+		NotificationPin->ExpireAndFadeout();
+		RestartNotificationItem.Reset();
+	}
+}
+
+
+void UEditorSysConfigAssistantSubsystem::AddIssue(const FEditorSysConfigIssue& Issue)
+{
+	FWriteScopeLock _(IssuesLock);
+	const TSharedPtr<FEditorSysConfigIssue>* ExistingIssue = Issues.FindByPredicate([&Issue](const TSharedPtr<FEditorSysConfigIssue>& ExistingIssue)
+		{
+			return ExistingIssue->Feature == Issue.Feature;
+		});
+
+	bool bIssueAddedOrUpgraded = false;
+	if (ExistingIssue)
+	{
+		if ((*ExistingIssue)->Severity < Issue.Severity)
+		{
+			**ExistingIssue = Issue;
+			bIssueAddedOrUpgraded = true;
+		}
+	}
+	else
+	{
+		Issues.Add(MakeShared<FEditorSysConfigIssue>(Issue));
+		bIssueAddedOrUpgraded = true;
+	}
+	if (bIssueAddedOrUpgraded && (Issue.Severity == EEditorSysConfigIssueSeverity::High))
+	{
+		if (IsInGameThread())
+		{
+			NotifySystemConfigIssues();
+		}
+		else
+		{
+			Async(EAsyncExecution::TaskGraphMainThread, NotifySystemConfigIssues);
+		}
+	}
+}
+
+TArray<TSharedPtr<FEditorSysConfigIssue>> UEditorSysConfigAssistantSubsystem::GetIssues()
+{
+	FReadScopeLock _(IssuesLock);
+	return Issues;
+}
+
+void UEditorSysConfigAssistantSubsystem::ApplySysConfigChanges(TArrayView<const TSharedPtr<FEditorSysConfigIssue>> InIssues)
+{
+	TArray<IEditorSysConfigFeature*> FeaturesToRecheck;
+	TArray<FString> ElevatedCommands;
+	EEditorSysConfigFeatureRemediationFlags AccumulatedFlags = EEditorSysConfigFeatureRemediationFlags::NoAutomatedRemediation;
+	for (const TSharedPtr<FEditorSysConfigIssue>& Issue : InIssues)
+	{
+		EEditorSysConfigFeatureRemediationFlags Flags = Issue->Feature->GetRemediationFlags();
+		if (EnumHasAnyFlags(Flags, EEditorSysConfigFeatureRemediationFlags::NoAutomatedRemediation))
+		{
+			continue;
+		}
+
+		AccumulatedFlags |= Flags;
+
+		Issue->Feature->ApplySysConfigChanges(ElevatedCommands);
+
+#if PLATFORM_WINDOWS
+		if (!EnumHasAnyFlags(Flags, EEditorSysConfigFeatureRemediationFlags::RequiresSystemRestart | EEditorSysConfigFeatureRemediationFlags::RequiresApplicationRestart))
+#else
+		if (!EnumHasAnyFlags(Flags, EEditorSysConfigFeatureRemediationFlags::RequiresApplicationRestart))
+#endif
+		{
+			{
+				FWriteScopeLock _(IssuesLock);
+				Issues.RemoveAll([Issue](const TSharedPtr<FEditorSysConfigIssue>& ExistingIssue)
+					{
+						return ExistingIssue->Feature == Issue->Feature;
+					});
+			}
+			FeaturesToRecheck.Add(Issue->Feature);
+		}
+	}
+
+	if (!ElevatedCommands.IsEmpty())
+	{
+#if PLATFORM_WINDOWS
+		const TCHAR* CommandFileBaseName = TEXT("UESysConfig.bat");
+#elif PLATFORM_MAC || PLATFORM_LINUX
+		const TCHAR* CommandFileBaseName = TEXT("UESysConfig.sh");
+#endif
+		const FString CommandFilename = FPaths::Combine(FPlatformProcess::UserTempDir(), CommandFileBaseName);
+		FFileHelper::SaveStringArrayToFile(ElevatedCommands, *CommandFilename);
+
+		FString URL;
+		FString Params;
+#if PLATFORM_WINDOWS
+		URL = TEXT("cmd.exe");
+		Params = FString::Printf(TEXT("/c \"\"%s\"\""), *CommandFilename);
+#elif PLATFORM_MAC || PLATFORM_LINUX
+		URL = TEXT("/usr/bin/env");
+		Params = FString::Printf(TEXT(" -- \"%s\""), *CommandFilename);
+#endif
+
+		int32 ExecutionReturnCode = 0;
+		FPlatformProcess::ExecElevatedProcess(*URL, *Params, &ExecutionReturnCode);
+	}
+
+	for (IEditorSysConfigFeature* FeatureToRecheck : FeaturesToRecheck)
+	{
+		FeatureToRecheck->StartSystemCheck();
+	}
+
+#if PLATFORM_WINDOWS
+	if (EnumHasAnyFlags(AccumulatedFlags, EEditorSysConfigFeatureRemediationFlags::RequiresSystemRestart))
+	{
+		NotifyRestart(/* bAppplicationOnly */false);
+		return;
+	}
+#endif
+	if (EnumHasAnyFlags(AccumulatedFlags, EEditorSysConfigFeatureRemediationFlags::RequiresApplicationRestart))
+	{
+		NotifyRestart(/* bAppplicationOnly */true);
+	}
+}
+
+void UEditorSysConfigAssistantSubsystem::DismissSystemConfigNotification()
+{
+	if (IssueNotificationItem.IsValid())
+	{
+		IssueNotificationItem->ExpireAndFadeout();
+	}
+}
+
+#undef LOCTEXT_NAMESPACE
diff --git a/Engine/Plugins/Editor/EditorSysConfigAssistant/Source/EditorSysConfigAssistant/Private/Widgets/EditorSysConfigAssistantDelegates.h b/Engine/Plugins/Editor/EditorSysConfigAssistant/Source/EditorSysConfigAssistant/Private/Widgets/EditorSysConfigAssistantDelegates.h
new file mode 100644
--- /dev/null
+++ b/Engine/Plugins/Editor/EditorSysConfigAssistant/Source/EditorSysConfigAssistant/Private/Widgets/EditorSysConfigAssistantDelegates.h
@@ -0,0 +1,15 @@
+// Copyright Epic Games, Inc. All Rights Reserved.
+
+#pragma once
+
+#include "Delegates/DelegateCombinations.h"
+#include "Templates/SharedPointer.h"
+
+struct FEditorSysConfigIssue;
+
+/**
+ * Delegate type for Editor System Configuration Assistant action running.
+ *
+ * The first parameter is the FEditorSysConfigIssue that you want to resolve issues on.
+ */
+DECLARE_DELEGATE_OneParam(FOnApplySysConfigChange, const TSharedPtr<FEditorSysConfigIssue>&)
diff --git a/Engine/Plugins/Editor/EditorSysConfigAssistant/Source/EditorSysConfigAssistant/Private/Widgets/SEditorSysConfigAssistant.cpp b/Engine/Plugins/Editor/EditorSysConfigAssistant/Source/EditorSysConfigAssistant/Private/Widgets/SEditorSysConfigAssistant.cpp
new file mode 100644
--- /dev/null
+++ b/Engine/Plugins/Editor/EditorSysConfigAssistant/Source/EditorSysConfigAssistant/Private/Widgets/SEditorSysConfigAssistant.cpp
@@ -0,0 +1,143 @@
+// Copyright Epic Games, Inc. All Rights Reserved.
+
+#include "Widgets/SEditorSysConfigAssistant.h"
+
+#include "Editor/EditorEngine.h"
+#include "EditorSysConfigAssistantSubsystem.h"
+#include "Framework/Commands/UIAction.h"
+#include "Framework/Docking/TabManager.h"
+#include "Framework/MultiBox/MultiBoxBuilder.h"
+#include "SlateOptMacros.h"
+#include "SPositiveActionButton.h"
+#include "SSimpleButton.h"
+#include "Styling/AppStyle.h"
+#include "Styling/SlateTypes.h"
+#include "Styling/StyleColors.h"
+#include "Widgets/Docking/SDockTab.h"
+#include "Widgets/Images/SImage.h"
+#include "Widgets/Input/SButton.h"
+#include "Widgets/Input/SCheckBox.h"
+#include "Widgets/Input/SComboButton.h"
+#include "Widgets/Layout/SBorder.h"
+#include "Widgets/Layout/SSplitter.h"
+#include "Widgets/Layout/SWidgetSwitcher.h"
+#include "Widgets/SBoxPanel.h"
+#include "Widgets/SEditorSysConfigAssistantIssueListView.h"
+#include "Widgets/Text/STextBlock.h"
+
+extern UNREALED_API UEditorEngine* GEditor;
+
+#define LOCTEXT_NAMESPACE "SEditorSysConfigAssistant"
+
+SEditorSysConfigAssistant::SEditorSysConfigAssistant()
+{
+}
+
+
+SEditorSysConfigAssistant::~SEditorSysConfigAssistant()
+{
+}
+
+BEGIN_SLATE_FUNCTION_BUILD_OPTIMIZATION
+void SEditorSysConfigAssistant::Construct(const FArguments& InArgs, const TSharedRef<SDockTab>& ConstructUnderMajorTab, const TSharedPtr<SWindow>& ConstructUnderWindow)
+{
+
+	ChildSlot
+	[
+		SNew(SVerticalBox)
+			
+		+ SVerticalBox::Slot()
+		.AutoHeight()
+		.Padding(0)
+		[
+			SNew(SBorder)
+			.BorderImage(FAppStyle::Get().GetBrush("Brushes.Panel"))
+			.Padding(FMargin(0, 6))
+			[
+				SNew(SHorizontalBox)
+
+				+ SHorizontalBox::Slot()
+				.FillWidth(0.8f)
+				.HAlign(HAlign_Fill)
+				.Padding(14, 0, 0, 0)
+				[
+					SNew(STextBlock)
+					.AutoWrapText(true)
+					.Text_Lambda([this]()
+					{
+						return LOCTEXT("IssuesPreamble", "The list below shows system configuration issues that are impacting editor experience.  A button next to each item can be used to fix the configuration issue, or you can press the 'Apply All Changes' button to address all the identified issues.");
+					})
+				]
+
+				+ SHorizontalBox::Slot()
+				.FillWidth(0.2f)
+				.HAlign(HAlign_Right)
+				.Padding(0.0f, 0.0f)
+				[
+					SNew(SSimpleButton)
+					.Text(LOCTEXT("ApplyAllChangesButton", "Apply All Changes"))
+					.OnClicked_Lambda([this]()
+					{
+						UEditorSysConfigAssistantSubsystem* Subsystem = GEditor->GetEditorSubsystem<UEditorSysConfigAssistantSubsystem>();
+						if (!Subsystem)
+						{
+							return FReply::Handled();
+						}
+
+						TArray<TSharedPtr<FEditorSysConfigIssue>> IssuesCopy = Subsystem->GetIssues();
+						Subsystem->ApplySysConfigChanges(IssuesCopy);
+						IssueListView->RefreshIssueList();
+
+						return FReply::Handled();
+					})
+					.IsEnabled_Lambda([this]()
+					{
+						UEditorSysConfigAssistantSubsystem* Subsystem = GEditor->GetEditorSubsystem<UEditorSysConfigAssistantSubsystem>();
+						if (!Subsystem)
+						{
+							return false;
+						}
+
+						return !Subsystem->GetIssues().IsEmpty();
+					})
+				]
+			]
+		]
+
+		+ SVerticalBox::Slot()
+		.FillHeight(1)
+		.Padding(0)
+		[
+			SNew(SBorder)
+			.BorderImage(FAppStyle::Get().GetBrush("Brushes.Recessed"))
+			[
+				SAssignNew(IssueListView, SEditorSysConfigAssistantIssueListView)
+				.OnApplySysConfigChange(this, &SEditorSysConfigAssistant::OnApplySysConfigChange)
+			]
+		]
+		
+	];
+
+	UEditorSysConfigAssistantSubsystem* Subsystem = GEditor->GetEditorSubsystem<UEditorSysConfigAssistantSubsystem>();
+	if (!Subsystem)
+	{
+		return;
+	}
+
+	Subsystem->DismissSystemConfigNotification();
+}
+END_SLATE_FUNCTION_BUILD_OPTIMIZATION
+
+void SEditorSysConfigAssistant::OnApplySysConfigChange(const TSharedPtr<FEditorSysConfigIssue>& Issue)
+{
+	UEditorSysConfigAssistantSubsystem* Subsystem = GEditor->GetEditorSubsystem<UEditorSysConfigAssistantSubsystem>();
+	if (!Subsystem)
+	{
+		return;
+	}
+
+	Subsystem->ApplySysConfigChanges({Issue});
+	IssueListView->RefreshIssueList();
+}
+
+#undef LOCTEXT_NAMESPACE
diff --git a/Engine/Plugins/Editor/EditorSysConfigAssistant/Source/EditorSysConfigAssistant/Private/Widgets/SEditorSysConfigAssistant.h b/Engine/Plugins/Editor/EditorSysConfigAssistant/Source/EditorSysConfigAssistant/Private/Widgets/SEditorSysConfigAssistant.h
new file mode 100644
--- /dev/null
+++ b/Engine/Plugins/Editor/EditorSysConfigAssistant/Source/EditorSysConfigAssistant/Private/Widgets/SEditorSysConfigAssistant.h
@@ -0,0 +1,59 @@
+// Copyright Epic Games, Inc. All Rights Reserved.
+
+#pragma once
+
+#include "CoreTypes.h"
+#include "EditorSysConfigIssue.h"
+#include "Framework/Docking/WorkspaceItem.h"
+#include "Input/Reply.h"
+#include "Layout/Visibility.h"
+#include "SlateFwd.h"
+#include "Templates/SharedPointer.h"
+#include "Widgets/DeclarativeSyntaxSupport.h"
+#include "Widgets/SCompoundWidget.h"
+
+class SEditorSysConfigAssistantIssueListView;
+
+enum class ECheckBoxState : uint8;
+
+
+/**
+ * Implements a Slate widget for the editor system configuration user interface.
+ */
+class SEditorSysConfigAssistant
+	: public SCompoundWidget
+{
+public:
+
+	SLATE_BEGIN_ARGS(SEditorSysConfigAssistant) { }
+	SLATE_END_ARGS()
+
+public:
+
+	/** Constructor. */
+	SEditorSysConfigAssistant();
+
+	/** Destructor. */
+	~SEditorSysConfigAssistant();
+
+public:
+
+	/**
+	 * Constructs the widget.
+	 *
+	 * @param InArgs The Slate argument list.
+	 * @param ConstructUnderMajorTab The major tab which will contain the session front-end.
+	 * @param ConstructUnderWindow The window in which this widget is being constructed.
+	 * @param InModel The view model to use.
+	 */
+	void Construct(const FArguments& InArgs, const TSharedRef<SDockTab>& ConstructUnderMajorTab, const TSharedPtr<SWindow>& ConstructUnderWindow);
+
+protected:
+
+private:
+	void OnApplySysConfigChange(const TSharedPtr<FEditorSysConfigIssue>& Issue);
+
+private:
+	TSharedPtr<SEditorSysConfigAssistantIssueListView> IssueListView;
+
+};
diff --git a/Engine/Plugins/Editor/EditorSysConfigAssistant/Source/EditorSysConfigAssistant/Private/Widgets/SEditorSysConfigAssistantIssueApplyButton.h b/Engine/Plugins/Editor/EditorSysConfigAssistant/Source/EditorSysConfigAssistant/Private/Widgets/SEditorSysConfigAssistantIssueApplyButton.h
new file mode 100644
--- /dev/null
+++ b/Engine/Plugins/Editor/EditorSysConfigAssistant/Source/EditorSysConfigAssistant/Private/Widgets/SEditorSysConfigAssistantIssueApplyButton.h
@@ -0,0 +1,88 @@
+// Copyright Epic Games, Inc. All Rights Reserved.
+
+#pragma once
+
+#include "CoreMinimal.h"
+#include "Layout/Visibility.h"
+#include "Misc/Attribute.h"
+#include "SSimpleButton.h"
+#include "Styling/AppStyle.h"
+#include "Widgets/DeclarativeSyntaxSupport.h"
+#include "Widgets/Images/SImage.h"
+#include "Widgets/Input/SButton.h"
+#include "Widgets/SBoxPanel.h"
+#include "Widgets/SCompoundWidget.h"
+#include "Widgets/SOverlay.h"
+#include "Widgets/SToolTip.h"
+#include "Widgets/SWidget.h"
+#include "Widgets/Text/STextBlock.h"
+
+#define LOCTEXT_NAMESPACE "SEditorSysConfigAssistantIssueApplyButton"
+
+/**
+ * Implements a button for applying config changes to resolve a system configuration issue.
+ */
+class SEditorSysConfigAssistantIssueApplyButton
+	: public SCompoundWidget
+{
+public:
+
+	SLATE_BEGIN_ARGS(SEditorSysConfigAssistantIssueApplyButton) { }
+		SLATE_EVENT(FOnClicked, OnClicked)
+	SLATE_END_ARGS()
+
+public:
+
+	void Construct(const FArguments& InArgs, bool InShowText)
+	{
+		TSharedPtr<SVerticalBox> VerticalBoxWidget;
+		ChildSlot
+		[	
+			SNew(SSimpleButton)
+			.OnClicked(InArgs._OnClicked)
+			.IsEnabled(this, &SEditorSysConfigAssistantIssueApplyButton::ButtonEnabled)
+			.Icon(this, &SEditorSysConfigAssistantIssueApplyButton::GetApplyIcon)
+		];
+
+		if (InShowText && VerticalBoxWidget.IsValid())
+		{
+			VerticalBoxWidget->AddSlot()
+			.AutoHeight()
+			[
+				SNew(STextBlock)
+				.Text(LOCTEXT("EditorSysConfigAssistant", "Apply Change"))
+			];
+		}
+
+		// Otherwise we fall back on simple text
+		SetToolTipText(LOCTEXT("EditorSysConfigAssistant", "Apply changes to address this issue"));
+	}
+
+	virtual TSharedPtr<IToolTip> GetToolTip() override
+	{
+		return SWidget::GetToolTip();
+	}
+
+private:
+
+	bool HasError() const
+	{
+		return false;
+	}
+
+	bool ButtonEnabled() const
+	{
+		return !HasError();
+	}
+
+	const FSlateBrush* GetApplyIcon() const
+	{
+		return HasError() ? FAppStyle::Get().GetBrush(TEXT("Icons.Error")) : FAppStyle::Get().GetBrush("Icons.Play");
+	}
+
+private:
+
+};
+
+
+#undef LOCTEXT_NAMESPACE
diff --git a/Engine/Plugins/Editor/EditorSysConfigAssistant/Source/EditorSysConfigAssistant/Private/Widgets/SEditorSysConfigAssistantIssueListRow.h b/Engine/Plugins/Editor/EditorSysConfigAssistant/Source/EditorSysConfigAssistant/Private/Widgets/SEditorSysConfigAssistantIssueListRow.h
new file mode 100644
--- /dev/null
+++ b/Engine/Plugins/Editor/EditorSysConfigAssistant/Source/EditorSysConfigAssistant/Private/Widgets/SEditorSysConfigAssistantIssueListRow.h
@@ -0,0 +1,213 @@
+// Copyright Epic Games, Inc. All Rights Reserved.
+
+#pragma once
+
+#include "CoreMinimal.h"
+#include "EditorSysConfigIssue.h"
+#include "Input/Reply.h"
+#include "Layout/Margin.h"
+#include "Layout/Visibility.h"
+#include "Misc/Attribute.h"
+#include "PlatformInfo.h"
+#include "Styling/AppStyle.h"
+#include "Styling/StyleDefaults.h"
+#include "Widgets/DeclarativeSyntaxSupport.h"
+#include "Widgets/EditorSysConfigAssistantDelegates.h"
+#include "Widgets/Images/SImage.h"
+#include "Widgets/Layout/SBorder.h"
+#include "Widgets/Layout/SBox.h"
+#include "Widgets/Layout/SSeparator.h"
+#include "Widgets/Layout/SUniformGridPanel.h"
+#include "Widgets/SBoxPanel.h"
+#include "Widgets/SEditorSysConfigAssistantIssueApplyButton.h"
+#include "Widgets/Text/STextBlock.h"
+#include "Widgets/Views/SListView.h"
+#include "Widgets/Views/STableRow.h"
+#include "Widgets/Views/STableViewBase.h"
+
+#define LOCTEXT_NAMESPACE "SEditorSysConfigAssistantIssueListRow"
+
+/**
+ * Implements a row widget for the editor's system configuration assistant issue list.
+ */
+class SEditorSysConfigAssistantIssueListRow
+	: public STableRow<TSharedPtr<FEditorSysConfigIssue>>
+{
+public:
+
+	SLATE_BEGIN_ARGS(SEditorSysConfigAssistantIssueListRow) { }
+		/**
+		* The Callback for when the button to resolve an issue is clicked.
+		*/
+		SLATE_EVENT(FOnApplySysConfigChange, OnApplySysConfigChange)
+	
+		/**
+		 * The issue shown in this row.
+		 */
+		SLATE_ARGUMENT(TSharedPtr<FEditorSysConfigIssue>, Issue)
+
+	SLATE_END_ARGS()
+
+public:
+
+	~SEditorSysConfigAssistantIssueListRow()
+	{
+	}
+
+	void Construct(const FArguments& InArgs, const TSharedRef<STableViewBase>& InOwnerTableView)
+	{
+		STableRow< TSharedPtr<FEditorSysConfigIssue> >::ConstructInternal(
+			STableRow::FArguments()
+			.ShowSelection(false),
+			InOwnerTableView
+			);
+
+		Issue = InArgs._Issue.ToSharedRef();
+		OnApplySysConfigChange = InArgs._OnApplySysConfigChange;
+
+		TSharedRef<SUniformGridPanel> NameGrid = SNew(SUniformGridPanel).SlotPadding(FMargin(0.0f, 1.0f));
+		TSharedRef<SUniformGridPanel> ValueGrid = SNew(SUniformGridPanel).SlotPadding(FMargin(0.0f, 1.0f));
+
+		ChildSlot
+		[
+			SNew(SHorizontalBox)
+
+			+ SHorizontalBox::Slot()
+			.FillWidth(1)
+			.Padding(0,0,0,1)
+			[
+				SNew(SBorder)
+				.Padding(2)
+				.BorderImage(FAppStyle::Get().GetBrush("Brushes.Panel"))
+				[
+					SNew(SHorizontalBox)
+
+					+ SHorizontalBox::Slot()
+					.AutoWidth()
+					.VAlign(VAlign_Center)
+					.Padding(14, 0, 12, 0)
+					[
+						SNew(SBox)
+						.WidthOverride(44.f)
+						.HeightOverride(44.f)
+						[
+							SNew(SImage)
+							.Image(this, &SEditorSysConfigAssistantIssueListRow::HandleIssueImage)
+						]
+					]
+
+					+ SHorizontalBox::Slot()
+					.FillWidth(1)
+					.VAlign(VAlign_Center)
+					[
+						SNew(SVerticalBox)
+
+						+ SVerticalBox::Slot()
+						.AutoHeight()
+						.Padding(2,9,2,4)
+						[
+							SNew(STextBlock)
+							.Text(this, &SEditorSysConfigAssistantIssueListRow::HandleIssueNameText)
+						]
+
+						+ SVerticalBox::Slot()
+						.AutoHeight()
+						.Padding(2, 4, 2, 9)
+						[
+							SNew(STextBlock)
+							.AutoWrapText(true)
+							.Text(this, &SEditorSysConfigAssistantIssueListRow::HandleIssueDescriptionText)
+						]
+					]
+
+					+ SHorizontalBox::Slot()
+					.AutoWidth()
+					.VAlign(VAlign_Center)
+					[
+						// This Vertical box ensures the NameGrid spans only the vertical space the ValueGrid forces.
+						SNew(SVerticalBox)
+
+						+ SVerticalBox::Slot()
+						.AutoHeight()
+						.Padding(0)
+						[
+							SNew(SHorizontalBox)
+
+							+ SHorizontalBox::Slot()
+							.AutoWidth()
+							.Padding(0, 0, 30, 0)
+							[
+								NameGrid
+							]
+
+							+ SHorizontalBox::Slot()
+							.AutoWidth()
+							[
+								ValueGrid
+							]
+						]
+					]
+
+					+ SHorizontalBox::Slot()
+					.AutoWidth()
+					.Padding(12, 5, 0, 5)
+					[
+						SNew(SSeparator)
+						.Orientation(Orient_Vertical)
+						.Thickness(1.f)
+					]
+
+					+ SHorizontalBox::Slot()
+					.AutoWidth()
+					.VAlign(VAlign_Center)
+					.Padding(20, 0, 20, 0)
+					[
+						SNew(SEditorSysConfigAssistantIssueApplyButton, true)
+						.OnClicked(this, &SEditorSysConfigAssistantIssueListRow::OnApplySysConfigChangeClicked)
+					]
+				]
+			]
+		];
+
+	}
+
+private:
+
+	FReply OnApplySysConfigChangeClicked()
+	{
+		if (OnApplySysConfigChange.IsBound())
+		{
+			OnApplySysConfigChange.Execute(Issue);
+		}
+
+		return FReply::Handled();
+	}
+
+	// Callback for getting the icon image of the issue.
+	const FSlateBrush* HandleIssueImage( ) const
+	{
+		return FStyleDefaults::GetNoBrush();
+	}
+
+	// Callback for getting the friendly name.
+	FText HandleIssueNameText() const
+	{
+		return Issue->Feature->GetDisplayName();
+	}
+
+	FText HandleIssueDescriptionText() const
+	{
+		return Issue->Feature->GetDisplayDescription();
+	}
+
+private:
+
+	// Holds a reference to the issue that is displayed in this row.
+	TSharedPtr<FEditorSysConfigIssue> Issue;
+
+	// Holds a delegate to be invoked when a change is applied.
+	FOnApplySysConfigChange OnApplySysConfigChange;
+};
+
+
+#undef LOCTEXT_NAMESPACE
diff --git a/Engine/Plugins/Editor/EditorSysConfigAssistant/Source/EditorSysConfigAssistant/Private/Widgets/SEditorSysConfigAssistantIssueListView.cpp b/Engine/Plugins/Editor/EditorSysConfigAssistant/Source/EditorSysConfigAssistant/Private/Widgets/SEditorSysConfigAssistantIssueListView.cpp
new file mode 100644
--- /dev/null
+++ b/Engine/Plugins/Editor/EditorSysConfigAssistant/Source/EditorSysConfigAssistant/Private/Widgets/SEditorSysConfigAssistantIssueListView.cpp
@@ -0,0 +1,95 @@
+// Copyright Epic Games, Inc. All Rights Reserved.
+
+#include "Widgets/SEditorSysConfigAssistantIssueListView.h"
+
+#include "Editor/EditorEngine.h"
+#include "EditorSysConfigAssistantSubsystem.h"
+#include "EditorSysConfigFeature.h"
+#include "Framework/Docking/TabManager.h"
+#include "SlateOptMacros.h"
+#include "Widgets/Input/SHyperlink.h"
+#include "Widgets/Layout/SScrollBorder.h"
+#include "Widgets/SEditorSysConfigAssistantIssueListRow.h"
+#include "Widgets/Text/STextBlock.h"
+#include "Widgets/Views/STableRow.h"
+#include "Widgets/Views/STableViewBase.h"
+
+#define LOCTEXT_NAMESPACE "SEditorSysConfigAssistantIssueListView"
+
+extern UNREALED_API UEditorEngine* GEditor;
+
+SEditorSysConfigAssistantIssueListView::~SEditorSysConfigAssistantIssueListView()
+{
+}
+
+BEGIN_SLATE_FUNCTION_BUILD_OPTIMIZATION
+void SEditorSysConfigAssistantIssueListView::Construct(const FArguments& InArgs)
+{
+	OnApplySysConfigChange = InArgs._OnApplySysConfigChange;
+
+	SAssignNew(IssueListView, SListView<TSharedPtr<FEditorSysConfigIssue>>)
+		.SelectionMode(ESelectionMode::None)
+		.ListItemsSource(&IssueList)
+		.OnGenerateRow(this, &SEditorSysConfigAssistantIssueListView::HandleIssueListViewGenerateRow)
+		.ItemHeight(16.0f);
+
+	ChildSlot
+	[
+		SNew(SVerticalBox)
+
+		+ SVerticalBox::Slot()
+		.FillHeight(1.0f)
+		.Padding(0)
+		[
+			SNew(SScrollBorder, IssueListView.ToSharedRef())
+			[
+				IssueListView.ToSharedRef()
+			]
+		]
+	];
+
+	RefreshIssueList();
+}
+END_SLATE_FUNCTION_BUILD_OPTIMIZATION
+
+void SEditorSysConfigAssistantIssueListView::RefreshIssueList()
+{
+	IssueList.Empty();
+
+	UEditorSysConfigAssistantSubsystem* Subsystem = GEditor->GetEditorSubsystem<UEditorSysConfigAssistantSubsystem>();
+	if (!Subsystem)
+	{
+		return;
+	}
+
+	IssueList = Subsystem->GetIssues();
+	IssueListView->RequestListRefresh();
+}
+
+bool SEditorSysConfigAssistantIssueListView::HandleIssueListRowIsEnabled(TSharedPtr<FEditorSysConfigIssue> Issue) const
+{
+	return true;
+}
+
+TSharedRef<ITableRow> SEditorSysConfigAssistantIssueListView::HandleIssueListViewGenerateRow(TSharedPtr<FEditorSysConfigIssue> InItem, const TSharedRef<STableViewBase>& OwnerTable) const
+{
+	return SNew(SEditorSysConfigAssistantIssueListRow, OwnerTable)
+		.OnApplySysConfigChange(OnApplySysConfigChange)
+		.Issue(InItem)
+		.IsEnabled(this, &SEditorSysConfigAssistantIssueListView::HandleIssueListRowIsEnabled, InItem);
+}
+
+
+void SEditorSysConfigAssistantIssueListView::HandleIssueAdded(const TSharedPtr<FEditorSysConfigIssue> AddedIssue)
+{
+	RefreshIssueList();
+}
+
+
+void SEditorSysConfigAssistantIssueListView::HandleIssueRemoved(const TSharedPtr<FEditorSysConfigIssue> AddedIssue)
+{
+	RefreshIssueList();
+}
+
+
+#undef LOCTEXT_NAMESPACE
diff --git a/Engine/Plugins/Editor/EditorSysConfigAssistant/Source/EditorSysConfigAssistant/Private/Widgets/SEditorSysConfigAssistantIssueListView.h b/Engine/Plugins/Editor/EditorSysConfigAssistant/Source/EditorSysConfigAssistant/Private/Widgets/SEditorSysConfigAssistantIssueListView.h
new file mode 100644
--- /dev/null
+++ b/Engine/Plugins/Editor/EditorSysConfigAssistant/Source/EditorSysConfigAssistant/Private/Widgets/SEditorSysConfigAssistantIssueListView.h
@@ -0,0 +1,70 @@
+// Copyright Epic Games, Inc. All Rights Reserved.
+
+#pragma once
+
+#include "Containers/Array.h"
+#include "Layout/Visibility.h"
+#include "Misc/Attribute.h"
+#include "Templates/SharedPointer.h"
+#include "Widgets/DeclarativeSyntaxSupport.h"
+#include "Widgets/EditorSysConfigAssistantDelegates.h"
+#include "Widgets/SCompoundWidget.h"
+#include "Widgets/Views/SListView.h"
+
+class ITableRow;
+class STableViewBase;
+struct FEditorSysConfigIssue;
+
+/**
+ * Implements the deployment targets panel.
+ */
+class SEditorSysConfigAssistantIssueListView
+	: public SCompoundWidget
+{
+public:
+
+	SLATE_BEGIN_ARGS(SEditorSysConfigAssistantIssueListView) { }
+		SLATE_EVENT(FOnApplySysConfigChange, OnApplySysConfigChange)
+	SLATE_END_ARGS()
+
+public:
+
+	/** Destructor. */
+	~SEditorSysConfigAssistantIssueListView();
+
+	/**
+	 * Constructs the widget.
+	 *
+	 * @param InArgs The Slate argument list.
+	 */
+	void Construct(const FArguments& InArgs);
+
+
+	/** Refresh the list of issues. */
+	void RefreshIssueList();
+
+private:
+
+	/** Callback for getting the enabled state of a issue list row. */
+	bool HandleIssueListRowIsEnabled(TSharedPtr<FEditorSysConfigIssue> Issue) const;
+
+	/** Callback for generating a row in the issue view. */
+	TSharedRef<ITableRow> HandleIssueListViewGenerateRow(TSharedPtr<FEditorSysConfigIssue> InItem, const TSharedRef<STableViewBase>& OwnerTable) const;
+
+	/** Callback for when an issue has been added to the issue list. */
+	void HandleIssueAdded(const TSharedPtr<FEditorSysConfigIssue> AddedIssue);
+
+	/** Callback for when an issue has been added to the issue list. */
+	void HandleIssueRemoved(const TSharedPtr<FEditorSysConfigIssue> RemovedIssue);
+
+private:
+
+	/** Holds the list of available issues. */
+	TArray<TSharedPtr<FEditorSysConfigIssue>> IssueList;
+
+	/** Holds the issue list view . */
+	TSharedPtr<SListView<TSharedPtr<FEditorSysConfigIssue>> > IssueListView;
+
+	/** Holds a delegate to be invoked when an issue applies a config change. */
+	FOnApplySysConfigChange OnApplySysConfigChange;
+};
diff --git a/Engine/Plugins/Editor/EditorSysConfigAssistant/Source/EditorSysConfigAssistant/Private/Windows/EditorSysConfigFeatureLastAccessTime.cpp b/Engine/Plugins/Editor/EditorSysConfigAssistant/Source/EditorSysConfigAssistant/Private/Windows/EditorSysConfigFeatureLastAccessTime.cpp
new file mode 100644
--- /dev/null
+++ b/Engine/Plugins/Editor/EditorSysConfigAssistant/Source/EditorSysConfigAssistant/Private/Windows/EditorSysConfigFeatureLastAccessTime.cpp
@@ -0,0 +1,80 @@
+// Copyright Epic Games, Inc. All Rights Reserved.
+
+#include "EditorSysConfigFeatureLastAccessTime.h"
+
+#include "Async/Async.h"
+#include "Editor/EditorEngine.h"
+#include "EditorSysConfigAssistantSubsystem.h"
+#include "Internationalization/Internationalization.h"
+
+#if PLATFORM_WINDOWS
+
+#include "Microsoft/MinimalWindowsApi.h"
+#include "Windows/WindowsPlatformMisc.h"
+#include "Windows/AllowWindowsPlatformTypes.h"
+THIRD_PARTY_INCLUDES_START
+#include <winreg.h>
+THIRD_PARTY_INCLUDES_END
+#include "Windows/HideWindowsPlatformTypes.h"
+
+extern UNREALED_API UEditorEngine* GEditor;
+
+FText FEditorSysConfigFeatureLastAccessTime::GetDisplayName() const
+{
+	return NSLOCTEXT("EditorSysConfigAssistant", "LastAccessTimeAssistantName", "Last Access Time Tracking");
+}
+
+FText FEditorSysConfigFeatureLastAccessTime::GetDisplayDescription() const
+{
+	return NSLOCTEXT("EditorSysConfigAssistant", "LastAccessTimeAssistantDescription",
+		"Last access time tracking is enabled on your system and this can slow down asset registry scans and other operations.  It is recommended that you disable last access time tracking system wide.");
+}
+
+FGuid FEditorSysConfigFeatureLastAccessTime::GetVersion() const
+{
+	static FGuid VersionGuid(0xe1f5a6ac, 0x86b24293, 0x8413c9af, 0xdbf49c89);
+	return VersionGuid;
+}
+
+EEditorSysConfigFeatureRemediationFlags FEditorSysConfigFeatureLastAccessTime::GetRemediationFlags() const
+{
+	return EEditorSysConfigFeatureRemediationFlags::HasAutomatedRemediation |
+		EEditorSysConfigFeatureRemediationFlags::RequiresElevation;
+}
+
+void FEditorSysConfigFeatureLastAccessTime::StartSystemCheck()
+{
+	Async(EAsyncExecution::TaskGraph, [this]()
+		{
+			UEditorSysConfigAssistantSubsystem* Subsystem = GEditor->GetEditorSubsystem<UEditorSysConfigAssistantSubsystem>();
+			if (!Subsystem)
+			{
+				return;
+			}
+
+			FString LastAccessUpdateValueString;
+			if (FWindowsPlatformMisc::QueryRegKey(HKEY_LOCAL_MACHINE, TEXT("SYSTEM\\CurrentControlSet\\Control\\FileSystem"), TEXT("NtfsDisableLastAccessUpdate"), LastAccessUpdateValueString))
+			{
+				int32 KeyValue = 0;
+				LexFromString(KeyValue, *LastAccessUpdateValueString);
+				if (KeyValue & 1)
+				{
+					// Lowest bit being set means last access time updates are disabled
+					return;
+				}
+			}
+
+			FEditorSysConfigIssue NewIssue;
+			NewIssue.Feature = this;
+			// TODO: The severity should only be High if the asset count for the project (as obtained from the Asset Registry) is greater than some threshold
+			//		 so that we're not putting notifications in front of people working on small projects.
+			NewIssue.Severity = EEditorSysConfigIssueSeverity::High;
+			Subsystem->AddIssue(NewIssue);
+		});
+}
+
+void FEditorSysConfigFeatureLastAccessTime::ApplySysConfigChanges(TArray<FString>& OutElevatedCommands)
+{
+	OutElevatedCommands.Add(TEXT("fsutil behavior set disableLastAccess 1"));
+}
+#endif // PLATFORM_WINDOWS
diff --git a/Engine/Plugins/Editor/EditorSysConfigAssistant/Source/EditorSysConfigAssistant/Private/Windows/EditorSysConfigFeatureLastAccessTime.h b/Engine/Plugins/Editor/EditorSysConfigAssistant/Source/EditorSysConfigAssistant/Private/Windows/EditorSysConfigFeatureLastAccessTime.h
new file mode 100644
--- /dev/null
+++ b/Engine/Plugins/Editor/EditorSysConfigAssistant/Source/EditorSysConfigAssistant/Private/Windows/EditorSysConfigFeatureLastAccessTime.h
@@ -0,0 +1,24 @@
+// Copyright Epic Games, Inc. All Rights Reserved.
+
+#pragma once
+
+#if PLATFORM_WINDOWS
+
+#include "EditorSysConfigFeature.h"
+#include "Templates/SharedPointer.h"
+
+class FEditorSysConfigFeatureLastAccessTime : public IEditorSysConfigFeature
+{
+public:
+	virtual ~FEditorSysConfigFeatureLastAccessTime() = default;
+
+	virtual FText GetDisplayName() const override;
+	virtual FText GetDisplayDescription() const override;
+	virtual FGuid GetVersion() const override;
+	virtual EEditorSysConfigFeatureRemediationFlags GetRemediationFlags() const override;
+
+	virtual void StartSystemCheck() override;
+	virtual void ApplySysConfigChanges(TArray<FString>& OutElevatedCommands) override;
+};
+
+#endif // PLATFORM_WINDOWS
\ No newline at end of file
diff --git a/Engine/Plugins/Editor/EditorSysConfigAssistant/Source/EditorSysConfigAssistant/Public/EditorSysConfigAssistantModule.h b/Engine/Plugins/Editor/EditorSysConfigAssistant/Source/EditorSysConfigAssistant/Public/EditorSysConfigAssistantModule.h
new file mode 100644
--- /dev/null
+++ b/Engine/Plugins/Editor/EditorSysConfigAssistant/Source/EditorSysConfigAssistant/Public/EditorSysConfigAssistantModule.h
@@ -0,0 +1,19 @@
+// Copyright Epic Games, Inc. All Rights Reserved.
+
+#pragma once
+
+#include "Modules/ModuleInterface.h"
+#include "Modules/ModuleManager.h"
+
+class IEditorSysConfigAssistantModule : public IModuleInterface
+{
+public:
+	static FORCEINLINE IEditorSysConfigAssistantModule& Get()
+	{
+		return FModuleManager::LoadModuleChecked<IEditorSysConfigAssistantModule>("EditorSysConfigAssistant");
+	}
+	
+	/** Ensures the system configuration assistant UI is presented to the user */
+	virtual void ShowSystemConfigAssistant() = 0;
+};
+
diff --git a/Engine/Plugins/Editor/EditorSysConfigAssistant/Source/EditorSysConfigAssistant/Public/EditorSysConfigAssistantSubsystem.h b/Engine/Plugins/Editor/EditorSysConfigAssistant/Source/EditorSysConfigAssistant/Public/EditorSysConfigAssistantSubsystem.h
new file mode 100644
--- /dev/null
+++ b/Engine/Plugins/Editor/EditorSysConfigAssistant/Source/EditorSysConfigAssistant/Public/EditorSysConfigAssistantSubsystem.h
@@ -0,0 +1,54 @@
+// Copyright Epic Games, Inc. All Rights Reserved.
+
+#pragma once
+
+#include "EditorSubsystem.h"
+#include "EditorSysConfigIssue.h"
+#include "Templates/SharedPointer.h"
+
+#include "EditorSysConfigAssistantSubsystem.generated.h"
+
+class IModularFeature;
+class SNotificationItem;
+
+UCLASS()
+class EDITORSYSCONFIGASSISTANT_API UEditorSysConfigAssistantSubsystem : public UEditorSubsystem
+{
+	GENERATED_BODY()
+
+public:
+	virtual void Initialize(FSubsystemCollectionBase& Collection) override;
+	virtual void Deinitialize() override;
+	virtual bool ShouldCreateSubsystem(UObject* Outer) const override;
+
+	/** Can be called on any thread */
+	void AddIssue(const FEditorSysConfigIssue& Issue);
+	/** Can be called on any thread */
+	TArray<TSharedPtr<FEditorSysConfigIssue>> GetIssues();
+
+	/** Must be called on the game thread */
+	void ApplySysConfigChanges(TArrayView<const TSharedPtr<FEditorSysConfigIssue>> Issues);
+
+	/** Must be called on the game thread */
+	void DismissSystemConfigNotification();
+
+private:
+	void HandleModularFeatureRegistered(const FName& InFeatureName, IModularFeature* InFeature);
+	void HandleModularFeatureUnregistered(const FName& InFeatureName, IModularFeature* InFeature);
+	
+	/** Must be called on the game thread */
+	static void NotifySystemConfigIssues();
+
+	/** Must be called on the game thread */
+	void NotifyRestart(bool bApplicationOnly);
+
+	void OnApplicationRestartClicked();
+	void OnSystemRestartClicked();
+	void OnRestartDismissClicked();
+
+	FRWLock IssuesLock;
+	TArray<TSharedPtr<FEditorSysConfigIssue>> Issues;
+
+	TSharedPtr<SNotificationItem> IssueNotificationItem;
+	TWeakPtr<SNotificationItem> RestartNotificationItem;
+};
diff --git a/Engine/Plugins/Editor/EditorSysConfigAssistant/Source/EditorSysConfigAssistant/Public/EditorSysConfigFeature.h b/Engine/Plugins/Editor/EditorSysConfigAssistant/Source/EditorSysConfigAssistant/Public/EditorSysConfigFeature.h
new file mode 100644
--- /dev/null
+++ b/Engine/Plugins/Editor/EditorSysConfigAssistant/Source/EditorSysConfigAssistant/Public/EditorSysConfigFeature.h
@@ -0,0 +1,48 @@
+// Copyright Epic Games, Inc. All Rights Reserved.
+
+#pragma once
+
+#include "Containers/Array.h"
+#include "CoreTypes.h"
+#include "Features/IModularFeature.h"
+#include "Internationalization/Text.h"
+#include "Misc/EnumClassFlags.h"
+#include "Misc/Guid.h"
+#include "UObject/NameTypes.h"
+
+
+enum class EEditorSysConfigFeatureRemediationFlags : uint32
+{
+	NoAutomatedRemediation		= 0,
+	HasAutomatedRemediation		= 1 << 0,
+	RequiresElevation			= 1 << 1,
+	RequiresApplicationRestart	= 1 << 2,
+#if PLATFORM_WINDOWS
+	// Windows specific support for system restart
+	RequiresSystemRestart		= 1 << 3,
+#endif
+};
+ENUM_CLASS_FLAGS(EEditorSysConfigFeatureRemediationFlags);
+
+class IEditorSysConfigFeature : public IModularFeature
+{
+public:
+	virtual ~IEditorSysConfigFeature() = default;
+
+	static FName GetModularFeatureName()
+	{
+		static FName FeatureName = FName(TEXT("EditorSysConfigFeature"));
+		return FeatureName;
+	}
+
+	virtual FText GetDisplayName() const = 0;
+	virtual FText GetDisplayDescription() const = 0;
+	virtual FGuid GetVersion() const = 0;
+	virtual EEditorSysConfigFeatureRemediationFlags GetRemediationFlags() const = 0;
+	
+	/** Called from the game thread, encouraged to perform work asynchronously */
+	virtual void StartSystemCheck() = 0;
+
+	/** Called from the game thread, must complete synchronously */
+	virtual void ApplySysConfigChanges(TArray<FString>& OutElevatedCommands) = 0;
+};
diff --git a/Engine/Plugins/Editor/EditorSysConfigAssistant/Source/EditorSysConfigAssistant/Public/EditorSysConfigIssue.h b/Engine/Plugins/Editor/EditorSysConfigAssistant/Source/EditorSysConfigAssistant/Public/EditorSysConfigIssue.h
new file mode 100644
--- /dev/null
+++ b/Engine/Plugins/Editor/EditorSysConfigAssistant/Source/EditorSysConfigAssistant/Public/EditorSysConfigIssue.h
@@ -0,0 +1,18 @@
+// Copyright Epic Games, Inc. All Rights Reserved.
+
+#pragma once
+
+#include "EditorSysConfigFeature.h"
+
+enum class EEditorSysConfigIssueSeverity : uint8
+{
+	High,
+	Medium,
+	Low,
+};
+
+struct FEditorSysConfigIssue
+{
+	IEditorSysConfigFeature* Feature;
+	EEditorSysConfigIssueSeverity Severity;
+};
