Fixes for conflicting LOCTEXT labels.

diff --git a/Engine/Plugins/Editor/EditorSysConfigAssistant/Source/EditorSysConfigAssistant/Private/EditorSysConfigAssistantSubsystem.cpp b/Engine/Plugins/Editor/EditorSysConfigAssistant/Source/EditorSysConfigAssistant/Private/EditorSysConfigAssistantSubsystem.cpp
--- a/Engine/Plugins/Editor/EditorSysConfigAssistant/Source/EditorSysConfigAssistant/Private/EditorSysConfigAssistantSubsystem.cpp
+++ b/Engine/Plugins/Editor/EditorSysConfigAssistant/Source/EditorSysConfigAssistant/Private/EditorSysConfigAssistantSubsystem.cpp
@@ -140,11 +140,11 @@ void UEditorSysConfigAssistantSubsystem::NotifyRestart(bool bApplicationOnly)
 	FText RestartText;
 	if (bApplicationOnly)
 	{
-		RestartText = LOCTEXT("RestartRequiredTitle", "Application restart required to apply new system configuration");
+		RestartText = LOCTEXT("ApplicationRestartRequiredTitle", "Application restart required to apply new system configuration");
 	}
 	else
 	{
-		RestartText = LOCTEXT("RestartRequiredTitle", "System restart required to apply new system configuration");
+		RestartText = LOCTEXT("SystemRestartRequiredTitle", "System restart required to apply new system configuration");
 	}
 	FNotificationInfo Info(RestartText);
 
diff --git a/Engine/Plugins/Editor/EditorSysConfigAssistant/Source/EditorSysConfigAssistant/Private/Widgets/SEditorSysConfigAssistantIssueApplyButton.h b/Engine/Plugins/Editor/EditorSysConfigAssistant/Source/EditorSysConfigAssistant/Private/Widgets/SEditorSysConfigAssistantIssueApplyButton.h
--- a/Engine/Plugins/Editor/EditorSysConfigAssistant/Source/EditorSysConfigAssistant/Private/Widgets/SEditorSysConfigAssistantIssueApplyButton.h
+++ b/Engine/Plugins/Editor/EditorSysConfigAssistant/Source/EditorSysConfigAssistant/Private/Widgets/SEditorSysConfigAssistantIssueApplyButton.h
@@ -50,12 +50,12 @@ public:
 			.AutoHeight()
 			[
 				SNew(STextBlock)
-				.Text(LOCTEXT("EditorSysConfigAssistant", "Apply Change"))
+				.Text(LOCTEXT("ApplyChangeButtonText", "Apply Change"))
 			];
 		}
 
 		// Otherwise we fall back on simple text
-		SetToolTipText(LOCTEXT("EditorSysConfigAssistant", "Apply changes to address this issue"));
+		SetToolTipText(LOCTEXT("ApplyChangeButtonTooltip", "Apply changes to address this issue"));
 	}
 
 	virtual TSharedPtr<IToolTip> GetToolTip() override
