Wording changes to clarify details in the System Config Assistant.  Layout and formatting changes can be done later to improve legibility.

diff --git a/Engine/Plugins/Editor/EditorSysConfigAssistant/Source/EditorSysConfigAssistant/Private/Widgets/SEditorSysConfigAssistant.cpp b/Engine/Plugins/Editor/EditorSysConfigAssistant/Source/EditorSysConfigAssistant/Private/Widgets/SEditorSysConfigAssistant.cpp
--- a/Engine/Plugins/Editor/EditorSysConfigAssistant/Source/EditorSysConfigAssistant/Private/Widgets/SEditorSysConfigAssistant.cpp
+++ b/Engine/Plugins/Editor/EditorSysConfigAssistant/Source/EditorSysConfigAssistant/Private/Widgets/SEditorSysConfigAssistant.cpp
@@ -65,7 +65,10 @@ void SEditorSysConfigAssistant::Construct(const FArguments& InArgs, const TShare
 					.AutoWrapText(true)
 					.Text_Lambda([this]()
 					{
-						return LOCTEXT("IssuesPreamble", "The list below shows system configuration issues that are impacting editor experience.  A button next to each item can be used to fix the configuration issue, or you can press the 'Apply All Changes' button to address all the identified issues.");
+						return LOCTEXT("IssuesPreamble", "The list below shows system configuration issues that are impacting editor experience.\n"
+							"These are machine and operating system issues that exist outside of the editor, and addressing them can have consequences to software outside of the editor."
+							" If you are uncertain about applying the changes or lack permission to make the changes, please discuss with your system administrator."
+							" A button next to each item can be used to fix the configuration issue, or you can press the 'Apply All Changes' button to address all the identified issues.");
 					})
 				]
 
diff --git a/Engine/Plugins/Editor/EditorSysConfigAssistant/Source/EditorSysConfigAssistant/Private/Windows/EditorSysConfigFeatureLastAccessTime.cpp b/Engine/Plugins/Editor/EditorSysConfigAssistant/Source/EditorSysConfigAssistant/Private/Windows/EditorSysConfigFeatureLastAccessTime.cpp
--- a/Engine/Plugins/Editor/EditorSysConfigAssistant/Source/EditorSysConfigAssistant/Private/Windows/EditorSysConfigFeatureLastAccessTime.cpp
+++ b/Engine/Plugins/Editor/EditorSysConfigAssistant/Source/EditorSysConfigAssistant/Private/Windows/EditorSysConfigFeatureLastAccessTime.cpp
@@ -27,7 +27,7 @@ FText FEditorSysConfigFeatureLastAccessTime::GetDisplayName() const
 FText FEditorSysConfigFeatureLastAccessTime::GetDisplayDescription() const
 {
 	return NSLOCTEXT("EditorSysConfigAssistant", "LastAccessTimeAssistantDescription",
-		"Last access time tracking is enabled on your system and this can slow down asset registry scans and other operations.  It is recommended that you disable last access time tracking system wide.");
+		"Last access time tracking is enabled on your system and this can slow down asset registry scans and other operations. It is recommended that you disable last access time tracking machine-wide. Previous versions of Windows have had this disabled by default.");
 }
 
 FGuid FEditorSysConfigFeatureLastAccessTime::GetVersion() const
