Replace use of ExitWindowsEx with InitiateSystemShutdownEx, as the former is legacy and leads to a static analysis warning.

diff --git a/Engine/Plugins/Editor/EditorSysConfigAssistant/Source/EditorSysConfigAssistant/Private/EditorSysConfigAssistantSubsystem.cpp b/Engine/Plugins/Editor/EditorSysConfigAssistant/Source/EditorSysConfigAssistant/Private/EditorSysConfigAssistantSubsystem.cpp
--- a/Engine/Plugins/Editor/EditorSysConfigAssistant/Source/EditorSysConfigAssistant/Private/EditorSysConfigAssistantSubsystem.cpp
+++ b/Engine/Plugins/Editor/EditorSysConfigAssistant/Source/EditorSysConfigAssistant/Private/EditorSysConfigAssistantSubsystem.cpp
@@ -207,7 +207,7 @@ void UEditorSysConfigAssistantSubsystem::OnSystemRestartClicked()
 	}
 
 #if PLATFORM_WINDOWS
-	::ExitWindowsEx(EWX_REBOOT | EWX_FORCE, SHTDN_REASON_MINOR_MAINTENANCE | SHTDN_REASON_FLAG_PLANNED);
+	::InitiateSystemShutdownEx(nullptr, nullptr, 0, false, true, SHTDN_REASON_MAJOR_OTHER | SHTDN_REASON_MINOR_MAINTENANCE | SHTDN_REASON_FLAG_PLANNED);
 #endif
 	RequestEngineExit(TEXT("Restarting system to apply system config changes"));
 
