Suppress MSVC static analysis warning C28159 because it suggests against rebooting at all.  The choice to reboot in this case is not from the code itself but the system config it is influencing.

diff --git a/Engine/Plugins/Editor/EditorSysConfigAssistant/Source/EditorSysConfigAssistant/Private/EditorSysConfigAssistantSubsystem.cpp b/Engine/Plugins/Editor/EditorSysConfigAssistant/Source/EditorSysConfigAssistant/Private/EditorSysConfigAssistantSubsystem.cpp
--- a/Engine/Plugins/Editor/EditorSysConfigAssistant/Source/EditorSysConfigAssistant/Private/EditorSysConfigAssistantSubsystem.cpp
+++ b/Engine/Plugins/Editor/EditorSysConfigAssistant/Source/EditorSysConfigAssistant/Private/EditorSysConfigAssistantSubsystem.cpp
@@ -207,6 +207,7 @@ void UEditorSysConfigAssistantSubsystem::OnSystemRestartClicked()
 	}
 
 #if PLATFORM_WINDOWS
+	CA_SUPPRESS(28159);
 	::InitiateSystemShutdownEx(nullptr, nullptr, 0, false, true, SHTDN_REASON_MAJOR_OTHER | SHTDN_REASON_MINOR_MAINTENANCE | SHTDN_REASON_FLAG_PLANNED);
 #endif
 	RequestEngineExit(TEXT("Restarting system to apply system config changes"));
