Ensure that the EditorSysConfigAssistant subsystem delays its initialization until after OnAllModuleLoadingPhasesComplete and uses the tab manager's permission list to determine if it can present the system config UI and therefore it should evaluate system checks (and present notifications about them).

diff --git a/Engine/Plugins/Editor/EditorSysConfigAssistant/Source/EditorSysConfigAssistant/Private/EditorSysConfigAssistantModule.cpp b/Engine/Plugins/Editor/EditorSysConfigAssistant/Source/EditorSysConfigAssistant/Private/EditorSysConfigAssistantModule.cpp
--- a/Engine/Plugins/Editor/EditorSysConfigAssistant/Source/EditorSysConfigAssistant/Private/EditorSysConfigAssistantModule.cpp
+++ b/Engine/Plugins/Editor/EditorSysConfigAssistant/Source/EditorSysConfigAssistant/Private/EditorSysConfigAssistantModule.cpp
@@ -4,6 +4,7 @@
 
 #include "Features/IModularFeatures.h"
 #include "Framework/Docking/TabManager.h"
+#include "Misc/NamePermissionList.h"
 #include "Styling/AppStyle.h"
 #include "Textures/SlateIcon.h"
 #include "Widgets/DeclarativeSyntaxSupport.h"
@@ -29,6 +30,7 @@ public:
 	virtual void ShutdownModule() override;
 	// ~End IModuleInterface
 
+	virtual bool CanShowSystemConfigAssistant() override;
 	virtual void ShowSystemConfigAssistant() override;
 private:
 #if PLATFORM_WINDOWS
@@ -66,6 +68,11 @@ void FEditorSysConfigAssistantModule::ShutdownModule()
 #endif // PLATFORM_WINDOWS
 }
 
+bool FEditorSysConfigAssistantModule::CanShowSystemConfigAssistant()
+{
+	return FGlobalTabmanager::Get()->GetTabPermissionList()->PassesFilter(EditorSysConfigAssistantTabName);
+}
+
 void FEditorSysConfigAssistantModule::ShowSystemConfigAssistant()
 {
 	FGlobalTabmanager::Get()->TryInvokeTab(EditorSysConfigAssistantTabName);
diff --git a/Engine/Plugins/Editor/EditorSysConfigAssistant/Source/EditorSysConfigAssistant/Private/EditorSysConfigAssistantSubsystem.cpp b/Engine/Plugins/Editor/EditorSysConfigAssistant/Source/EditorSysConfigAssistant/Private/EditorSysConfigAssistantSubsystem.cpp
--- a/Engine/Plugins/Editor/EditorSysConfigAssistant/Source/EditorSysConfigAssistant/Private/EditorSysConfigAssistantSubsystem.cpp
+++ b/Engine/Plugins/Editor/EditorSysConfigAssistant/Source/EditorSysConfigAssistant/Private/EditorSysConfigAssistantSubsystem.cpp
@@ -11,6 +11,7 @@
 #include "Internationalization/Internationalization.h"
 #include "Internationalization/Text.h"
 #include "Misc/App.h"
+#include "Misc/CoreDelegates.h"
 #include "Misc/FileHelper.h"
 #include "Misc/Paths.h"
 #include "UnrealEdMisc.h"
@@ -54,20 +55,12 @@ void UEditorSysConfigAssistantSubsystem::Initialize(FSubsystemCollectionBase& Co
 {
 	Super::Initialize(Collection);
 
-	IModularFeatures& ModularFeatures = IModularFeatures::Get();
-	IModularFeatures::FScopedLockModularFeatureList LockModularFeatureList;
-	TArray<IEditorSysConfigFeature*> SysConfigFeatures = IModularFeatures::Get().GetModularFeatureImplementations<IEditorSysConfigFeature>(IEditorSysConfigFeature::GetModularFeatureName());
-	for (IEditorSysConfigFeature* SysConfigFeature : SysConfigFeatures)
-	{
-		SysConfigFeature->StartSystemCheck();
-	}
-
-	ModularFeatures.OnModularFeatureRegistered().AddUObject(this, &UEditorSysConfigAssistantSubsystem::HandleModularFeatureRegistered);
-	ModularFeatures.OnModularFeatureUnregistered().AddUObject(this, &UEditorSysConfigAssistantSubsystem::HandleModularFeatureUnregistered);
+	FCoreDelegates::OnAllModuleLoadingPhasesComplete.AddUObject(this,&UEditorSysConfigAssistantSubsystem::HandleAssistantInitializationEvent);
 }
 
 void UEditorSysConfigAssistantSubsystem::Deinitialize()
 {
+	FCoreDelegates::OnAllModuleLoadingPhasesComplete.RemoveAll(this);
 	Super::Deinitialize();
 
 	IModularFeatures& ModularFeatures = IModularFeatures::Get();
@@ -97,6 +90,23 @@ void UEditorSysConfigAssistantSubsystem::HandleModularFeatureUnregistered(const
 	}
 }
 
+void UEditorSysConfigAssistantSubsystem::HandleAssistantInitializationEvent()
+{
+	if (IEditorSysConfigAssistantModule::Get().CanShowSystemConfigAssistant())
+	{
+		IModularFeatures& ModularFeatures = IModularFeatures::Get();
+		IModularFeatures::FScopedLockModularFeatureList LockModularFeatureList;
+		TArray<IEditorSysConfigFeature*> SysConfigFeatures = IModularFeatures::Get().GetModularFeatureImplementations<IEditorSysConfigFeature>(IEditorSysConfigFeature::GetModularFeatureName());
+		for (IEditorSysConfigFeature* SysConfigFeature : SysConfigFeatures)
+		{
+			SysConfigFeature->StartSystemCheck();
+		}
+
+		ModularFeatures.OnModularFeatureRegistered().AddUObject(this, &UEditorSysConfigAssistantSubsystem::HandleModularFeatureRegistered);
+		ModularFeatures.OnModularFeatureUnregistered().AddUObject(this, &UEditorSysConfigAssistantSubsystem::HandleModularFeatureUnregistered);
+	}
+}
+
 void UEditorSysConfigAssistantSubsystem::NotifySystemConfigIssues()
 {
 	check(IsInGameThread());
diff --git a/Engine/Plugins/Editor/EditorSysConfigAssistant/Source/EditorSysConfigAssistant/Public/EditorSysConfigAssistantModule.h b/Engine/Plugins/Editor/EditorSysConfigAssistant/Source/EditorSysConfigAssistant/Public/EditorSysConfigAssistantModule.h
--- a/Engine/Plugins/Editor/EditorSysConfigAssistant/Source/EditorSysConfigAssistant/Public/EditorSysConfigAssistantModule.h
+++ b/Engine/Plugins/Editor/EditorSysConfigAssistant/Source/EditorSysConfigAssistant/Public/EditorSysConfigAssistantModule.h
@@ -12,8 +12,11 @@ public:
 	{
 		return FModuleManager::LoadModuleChecked<IEditorSysConfigAssistantModule>("EditorSysConfigAssistant");
 	}
-	
-	/** Ensures the system configuration assistant UI is presented to the user */
+
+	/** Checks if it is possible to show the system configuration assistant UI to the user */
+	virtual bool CanShowSystemConfigAssistant() = 0;
+
+	/** Attempts to present the system configuration assistant UI to the user */
 	virtual void ShowSystemConfigAssistant() = 0;
 };
 
diff --git a/Engine/Plugins/Editor/EditorSysConfigAssistant/Source/EditorSysConfigAssistant/Public/EditorSysConfigAssistantSubsystem.h b/Engine/Plugins/Editor/EditorSysConfigAssistant/Source/EditorSysConfigAssistant/Public/EditorSysConfigAssistantSubsystem.h
--- a/Engine/Plugins/Editor/EditorSysConfigAssistant/Source/EditorSysConfigAssistant/Public/EditorSysConfigAssistantSubsystem.h
+++ b/Engine/Plugins/Editor/EditorSysConfigAssistant/Source/EditorSysConfigAssistant/Public/EditorSysConfigAssistantSubsystem.h
@@ -36,6 +36,8 @@ private:
 	void HandleModularFeatureRegistered(const FName& InFeatureName, IModularFeature* InFeature);
 	void HandleModularFeatureUnregistered(const FName& InFeatureName, IModularFeature* InFeature);
 	
+	void HandleAssistantInitializationEvent();
+
 	/** Must be called on the game thread */
 	static void NotifySystemConfigIssues();
 
