[EditorSysConfigAssistant]

Allows the editor to check if the volume the user's project is running on has at least a 1 GiB journal created. If not, the assistant will run create one for the user if they accept the configuration wizard settings.

#rnx

diff --git a/Engine/Plugins/Editor/EditorSysConfigAssistant/Source/EditorSysConfigAssistant/Private/EditorSysConfigAssistantModule.cpp b/Engine/Plugins/Editor/EditorSysConfigAssistant/Source/EditorSysConfigAssistant/Private/EditorSysConfigAssistantModule.cpp
--- a/Engine/Plugins/Editor/EditorSysConfigAssistant/Source/EditorSysConfigAssistant/Private/EditorSysConfigAssistantModule.cpp
+++ b/Engine/Plugins/Editor/EditorSysConfigAssistant/Source/EditorSysConfigAssistant/Private/EditorSysConfigAssistantModule.cpp
@@ -15,6 +15,7 @@
 
 #if PLATFORM_WINDOWS
 #include "Windows/EditorSysConfigFeatureLastAccessTime.h"
+#include "Windows/EditorSysConfigFeatureUSNJournal.h"
 #endif
 
 static const FName EditorSysConfigAssistantTabName("EditorSysConfigAssistant");
@@ -35,6 +36,7 @@ public:
 private:
 #if PLATFORM_WINDOWS
 	FEditorSysConfigFeatureLastAccessTime LastAccessTimeFeature;
+	FEditorSysConfigFeatureUSNJournal USNJournalFeature;
 #endif // PLATFORM_WINDOWS
 };
 
@@ -58,12 +60,14 @@ void FEditorSysConfigAssistantModule::StartupModule()
 
 #if PLATFORM_WINDOWS
 	IModularFeatures::Get().RegisterModularFeature(LastAccessTimeFeature.GetModularFeatureName(), &LastAccessTimeFeature);
+	IModularFeatures::Get().RegisterModularFeature(USNJournalFeature.GetModularFeatureName(), &USNJournalFeature);
 #endif // PLATFORM_WINDOWS
 }
 
 void FEditorSysConfigAssistantModule::ShutdownModule()
 {
 #if PLATFORM_WINDOWS
+	IModularFeatures::Get().UnregisterModularFeature(USNJournalFeature.GetModularFeatureName(), &USNJournalFeature);
 	IModularFeatures::Get().UnregisterModularFeature(LastAccessTimeFeature.GetModularFeatureName(), &LastAccessTimeFeature);
 #endif // PLATFORM_WINDOWS
 }
diff --git a/Engine/Plugins/Editor/EditorSysConfigAssistant/Source/EditorSysConfigAssistant/Private/Windows/EditorSysConfigFeatureUSNJournal.cpp b/Engine/Plugins/Editor/EditorSysConfigAssistant/Source/EditorSysConfigAssistant/Private/Windows/EditorSysConfigFeatureUSNJournal.cpp
new file mode 100644
--- /dev/null
+++ b/Engine/Plugins/Editor/EditorSysConfigAssistant/Source/EditorSysConfigAssistant/Private/Windows/EditorSysConfigFeatureUSNJournal.cpp
@@ -0,0 +1,85 @@
+// Copyright Epic Games, Inc. All Rights Reserved.
+
+#include "EditorSysConfigFeatureUSNJournal.h"
+
+#include "Async/Async.h"
+#include "Editor/EditorEngine.h"
+#include "EditorSysConfigAssistantSubsystem.h"
+#include "HAL/PlatformFileManager.h"
+#include "Internationalization/Internationalization.h"
+
+#if PLATFORM_WINDOWS
+
+#include "Microsoft/MinimalWindowsApi.h"
+#include "Windows/WindowsPlatformMisc.h"
+#include "Windows/AllowWindowsPlatformTypes.h"
+THIRD_PARTY_INCLUDES_START
+#include <winreg.h>
+THIRD_PARTY_INCLUDES_END
+#include "Windows/HideWindowsPlatformTypes.h"
+
+extern UNREALED_API UEditorEngine* GEditor;
+
+const uint64 RecommendedJournalSizeBytes = 1 * 1024 * 1024 * 1024; /*1GiB*/
+FText FEditorSysConfigFeatureUSNJournal::GetDisplayName() const
+{
+	return NSLOCTEXT("EditorSysConfigAssistant", "USNJournalAssistantName", "USN Journal Configuration");
+}
+
+FText FEditorSysConfigFeatureUSNJournal::GetDisplayDescription() const
+{
+	return NSLOCTEXT("EditorSysConfigAssistant", "USNJournalAssistantDescription",
+		"The Update Sequence Number (USN) Journal is not configured correctly for the drive your project is stored on. This will result in slower asset discovery operations slowing down your Editor startup time. It is recommended that you create a 1 GiB USN journal on each drive you plan to store project files on.");
+}
+
+FGuid FEditorSysConfigFeatureUSNJournal::GetVersion() const
+{
+	static FGuid VersionGuid(0x52a34c4c, 0x93bb42c6, 0x96b5b787, 0x69b42f1c);
+
+	return VersionGuid;
+}
+
+EEditorSysConfigFeatureRemediationFlags FEditorSysConfigFeatureUSNJournal::GetRemediationFlags() const
+{
+	return EEditorSysConfigFeatureRemediationFlags::HasAutomatedRemediation |
+		EEditorSysConfigFeatureRemediationFlags::RequiresElevation;
+}
+
+void FEditorSysConfigFeatureUSNJournal::StartSystemCheck()
+{
+	Async(EAsyncExecution::TaskGraph, [this]()
+		{
+			UEditorSysConfigAssistantSubsystem* Subsystem = GEditor->GetEditorSubsystem<UEditorSysConfigAssistantSubsystem>();
+			if (!Subsystem)
+			{
+				return;
+			}
+
+			FString ProjectDir = FPaths::ProjectDir();
+			IPlatformFile& PlatformFile = FPlatformFileManager::Get().GetPlatformFile();
+			VolumeName = PlatformFile.FileJournalGetVolumeName(ProjectDir);
+			if (VolumeName.Len() < 2 || !VolumeName.EndsWith(TEXT(":")))
+			{
+				// If the volume isn't a root drive such (e.g. a network path, //mount/) simply return; there is nothing more we can do.
+				return;
+			}
+
+			uint64 JournalMaximumSize = PlatformFile.FileJournalGetMaximumSize(*VolumeName);	
+			if (JournalMaximumSize >= RecommendedJournalSizeBytes)
+			{
+				// Journal exists and is at or larger than the recommended amount
+				return;
+			}
+
+			FEditorSysConfigIssue NewIssue;
+			NewIssue.Feature = this;
+			NewIssue.Severity = EEditorSysConfigIssueSeverity::High;
+			Subsystem->AddIssue(NewIssue);
+		});
+}
+
+void FEditorSysConfigFeatureUSNJournal::ApplySysConfigChanges(TArray<FString>& OutElevatedCommands)
+{
+	OutElevatedCommands.Add(FString::Printf(TEXT("fsutil usn createjournal %s -m=%d"), *VolumeName, RecommendedJournalSizeBytes));
+}
+#endif // PLATFORM_WINDOWS
diff --git a/Engine/Plugins/Editor/EditorSysConfigAssistant/Source/EditorSysConfigAssistant/Private/Windows/EditorSysConfigFeatureUSNJournal.h b/Engine/Plugins/Editor/EditorSysConfigAssistant/Source/EditorSysConfigAssistant/Private/Windows/EditorSysConfigFeatureUSNJournal.h
new file mode 100644
--- /dev/null
+++ b/Engine/Plugins/Editor/EditorSysConfigAssistant/Source/EditorSysConfigAssistant/Private/Windows/EditorSysConfigFeatureUSNJournal.h
@@ -0,0 +1,27 @@
+// Copyright Epic Games, Inc. All Rights Reserved.
+
+#pragma once
+
+#if PLATFORM_WINDOWS
+
+#include "EditorSysConfigFeature.h"
+#include "Templates/SharedPointer.h"
+
+class FEditorSysConfigFeatureUSNJournal : public IEditorSysConfigFeature
+{
+public:
+	virtual ~FEditorSysConfigFeatureUSNJournal() = default;
+
+	virtual FText GetDisplayName() const override;
+	virtual FText GetDisplayDescription() const override;
+	virtual FGuid GetVersion() const override;
+	virtual EEditorSysConfigFeatureRemediationFlags GetRemediationFlags() const override;
+
+	virtual void StartSystemCheck() override;
+	virtual void ApplySysConfigChanges(TArray<FString>& OutElevatedCommands) override;
+
+private:
+	FString VolumeName;
+};
+
+#endif // PLATFORM_WINDOWS
\ No newline at end of file
