[EditorSysConfigAssistant]

- Corrects createjournal command where m=<size> should not have a leading `-`
- Ensures we print a uint64 correctly for all platforms

#rnx

diff --git a/Engine/Plugins/Editor/EditorSysConfigAssistant/Source/EditorSysConfigAssistant/Private/Windows/EditorSysConfigFeatureUSNJournal.cpp b/Engine/Plugins/Editor/EditorSysConfigAssistant/Source/EditorSysConfigAssistant/Private/Windows/EditorSysConfigFeatureUSNJournal.cpp
--- a/Engine/Plugins/Editor/EditorSysConfigAssistant/Source/EditorSysConfigAssistant/Private/Windows/EditorSysConfigFeatureUSNJournal.cpp
+++ b/Engine/Plugins/Editor/EditorSysConfigAssistant/Source/EditorSysConfigAssistant/Private/Windows/EditorSysConfigFeatureUSNJournal.cpp
@@ -80,6 +80,6 @@ void FEditorSysConfigFeatureUSNJournal::StartSystemCheck()
 
 void FEditorSysConfigFeatureUSNJournal::ApplySysConfigChanges(TArray<FString>& OutElevatedCommands)
 {
-	OutElevatedCommands.Add(FString::Printf(TEXT("fsutil usn createjournal %s -m=%d"), *VolumeName, RecommendedJournalSizeBytes));
+	OutElevatedCommands.Add(FString::Printf(TEXT("fsutil usn createjournal %s m=%" UINT64_FMT ""), *VolumeName, RecommendedJournalSizeBytes));
 }
 #endif // PLATFORM_WINDOWS
