Fix parsing of NtfsDisableLastAccessUpdate for EditorSysConfigAssistant

Numbers outside int32 range, e.g. 2147483650 (0x80000002) are being parsed as 2147483647 (0x7FFFFFFF) and need to be parsed as uint32 to be processed correctly. The low bit was erroneously being detected as set before the fix, which caused the plugin to falsely assume the setting was already good.

diff --git a/Engine/Plugins/Editor/EditorSysConfigAssistant/Source/EditorSysConfigAssistant/Private/Windows/EditorSysConfigFeatureLastAccessTime.cpp b/Engine/Plugins/Editor/EditorSysConfigAssistant/Source/EditorSysConfigAssistant/Private/Windows/EditorSysConfigFeatureLastAccessTime.cpp
--- a/Engine/Plugins/Editor/EditorSysConfigAssistant/Source/EditorSysConfigAssistant/Private/Windows/EditorSysConfigFeatureLastAccessTime.cpp
+++ b/Engine/Plugins/Editor/EditorSysConfigAssistant/Source/EditorSysConfigAssistant/Private/Windows/EditorSysConfigFeatureLastAccessTime.cpp
@@ -55,7 +55,7 @@ void FEditorSysConfigFeatureLastAccessTime::StartSystemCheck()
 			FString LastAccessUpdateValueString;
 			if (FWindowsPlatformMisc::QueryRegKey(HKEY_LOCAL_MACHINE, TEXT("SYSTEM\\CurrentControlSet\\Control\\FileSystem"), TEXT("NtfsDisableLastAccessUpdate"), LastAccessUpdateValueString))
 			{
-				int32 KeyValue = 0;
+				uint32 KeyValue = 0;
 				LexFromString(KeyValue, *LastAccessUpdateValueString);
 				if (KeyValue & 1)
 				{
